[LiveLinkPrestonMDR] Fix race condition on livelink source shutdown

diff --git a/Engine/Plugins/VirtualProduction/LiveLinkPrestonMDR/Source/LiveLinkPrestonMDR/Private/LiveLinkPrestonMDRSource.cpp b/Engine/Plugins/VirtualProduction/LiveLinkPrestonMDR/Source/LiveLinkPrestonMDR/Private/LiveLinkPrestonMDRSource.cpp
--- a/Engine/Plugins/VirtualProduction/LiveLinkPrestonMDR/Source/LiveLinkPrestonMDR/Private/LiveLinkPrestonMDRSource.cpp
+++ b/Engine/Plugins/VirtualProduction/LiveLinkPrestonMDR/Source/LiveLinkPrestonMDR/Private/LiveLinkPrestonMDRSource.cpp
@@ -15,21 +15,25 @@
 
 #define LOCTEXT_NAMESPACE "LiveLinkPrestonMDRSource"
 
-DEFINE_LOG_CATEGORY_STATIC(LogPrestonMDRSource, Log, All);
+DEFINE_LOG_CATEGORY_STATIC(LogLiveLinkPrestonMDRSource, Log, All);
 
 FLiveLinkPrestonMDRSource::FLiveLinkPrestonMDRSource(FLiveLinkPrestonMDRConnectionSettings InConnectionSettings)
 	: SocketSubsystem(ISocketSubsystem::Get(PLATFORM_SOCKETSUBSYSTEM))
 	, ConnectionSettings(MoveTemp(InConnectionSettings))
 	, LastTimeDataReceived(0.0)
-	, bIsConnectedToDevice(false)
-	, bFailedToConnectToDevice(false)
+	, MessageThreadShutdownTime(0.0)
+	, SourceStatus(EPrestonSourceStatus::NotConnected)
 {
 	SourceMachineName = FText::Format(LOCTEXT("PrestonMDRMachineName", "{0}:{1}"), FText::FromString(ConnectionSettings.IPAddress), FText::AsNumber(ConnectionSettings.PortNumber, &FNumberFormattingOptions::DefaultNoGrouping()));
 }
 
 FLiveLinkPrestonMDRSource::~FLiveLinkPrestonMDRSource()
 {
-	RequestSourceShutdown();
+	bool bIsReadyToShutdown = false;
+	while (!bIsReadyToShutdown)
+	{
+		bIsReadyToShutdown = ShutdownMessageThreadAndSocket();
+	}
 }
 
 void FLiveLinkPrestonMDRSource::ReceiveClient(ILiveLinkClient* InClient, FGuid InSourceGuid)
@@ -38,15 +42,6 @@ void FLiveLinkPrestonMDRSource::ReceiveClient(ILiveLinkClient* InClient, FGuid I
 	SubjectKey = FLiveLinkSubjectKey(InSourceGuid, ConnectionSettings.SubjectName);
 
 	OpenConnection();
-
-	MessageThread = MakeUnique<FPrestonMDRMessageThread>(Socket);
-
-	MessageThread->OnFrameDataReady().BindRaw(this, &FLiveLinkPrestonMDRSource::OnFrameDataReady);
-	MessageThread->OnStatusChanged().BindRaw(this, &FLiveLinkPrestonMDRSource::OnStatusChanged);
-	MessageThread->OnConnectionLost().BindRaw(this, &FLiveLinkPrestonMDRSource::OnConnectionLost);
-	MessageThread->OnConnectionFailed().BindRaw(this, &FLiveLinkPrestonMDRSource::OnConnectionFailed);
-
-	MessageThread->Start();
 }
 
 void FLiveLinkPrestonMDRSource::InitializeSettings(ULiveLinkSourceSettings* Settings)
@@ -64,7 +59,7 @@ void FLiveLinkPrestonMDRSource::InitializeSettings(ULiveLinkSourceSettings* Sett
 	}
 	else
 	{
-		UE_LOG(LogPrestonMDRSource, Warning, TEXT("Preston MDR Source coming from Preset is outdated. Consider recreating a Preston MDR Source. Configure it and resave as preset"));
+		UE_LOG(LogLiveLinkPrestonMDRSource, Warning, TEXT("Preston MDR Source coming from Preset is outdated. Consider recreating a Preston MDR Source. Configure it and resave as preset"));
 	}
 }
 
@@ -124,31 +119,60 @@ void FLiveLinkPrestonMDRSource::OnSettingsChanged(ULiveLinkSourceSettings* Setti
 		}
 		else
 		{
-			UE_LOG(LogPrestonMDRSource, Warning, TEXT("Preston MDR Source coming from Preset is outdated. Consider recreating a Preston MDR Source. Configure it and resave as preset"));
+			UE_LOG(LogLiveLinkPrestonMDRSource, Warning, TEXT("Preston MDR Source coming from Preset is outdated. Consider recreating a Preston MDR Source. Configure it and resave as preset"));
 		}
 	}
 }
 
 bool FLiveLinkPrestonMDRSource::IsSourceStillValid() const
 {
-	return bIsConnectedToDevice;
+	return (SourceStatus == EPrestonSourceStatus::ConnectedActive || SourceStatus == EPrestonSourceStatus::ConnectedIdle);
 }
 
 bool FLiveLinkPrestonMDRSource::RequestSourceShutdown()
 {
-	if (MessageThread)
 	{
-		MessageThread->Stop();
+		FScopeLock Lock(&PrestonSourceCriticalSection);
+		SourceStatus = EPrestonSourceStatus::ShuttingDown;
+	}
+
+	return ShutdownMessageThreadAndSocket();
+}
+
+bool FLiveLinkPrestonMDRSource::ShutdownMessageThreadAndSocket()
+{
+	if (MessageThread.IsValid())
+	{
+		// Instruct the message thread to stop its message loop
+		if (MessageThread->IsThreadRunning())
+		{
+			MessageThreadShutdownTime = FPlatformTime::Seconds();
+			MessageThread->Stop();
+		}
+
+		if ((FPlatformTime::Seconds() - MessageThreadShutdownTime) > MessageThreadShutdownTimeout)
+		{
+			MessageThread->ForceKill();
+		}
+
+		// If the message thread has stopped executing its message loop, then it is safe to reset it and destroy the socket
+		if (!MessageThread->IsFinished())
+		{
+			return false;
+		}
+
 		MessageThread.Reset();
 	}
 
+	if (Socket)
 	{
-		FScopeLock Lock(&PrestonSourceCriticalSection);
-		if (Socket)
+		if (!Socket->Close())
 		{
-			SocketSubsystem->DestroySocket(Socket);
-			Socket = nullptr;
+			UE_LOG(LogLiveLinkPrestonMDRSource, Warning, TEXT("Preston MDR Source socket failed to close."));
 		}
+
+		SocketSubsystem->DestroySocket(Socket);
+		Socket = nullptr;
 	}
 
 	return true;
@@ -159,21 +183,65 @@ FText FLiveLinkPrestonMDRSource::GetSourceType() const
 	return LOCTEXT("PrestonMDRSourceType", "PrestonMDR");
 }
 
+void FLiveLinkPrestonMDRSource::Update()
+{
+	if (SourceStatus == EPrestonSourceStatus::ConnectedActive && (FPlatformTime::Seconds() - LastTimeDataReceived > DataReceivedTimeout))
+	{
+		{
+			FScopeLock Lock(&PrestonSourceCriticalSection);
+			SourceStatus = EPrestonSourceStatus::ConnectedIdle;
+		}
+	}
+	else if (SourceStatus == EPrestonSourceStatus::ConnectionLost)
+	{
+		// If the connection was lost, shut down the message thread and the socket and attempt to reconnect
+		if (ShutdownMessageThreadAndSocket())
+		{
+			SourceStatus = EPrestonSourceStatus::NotConnected;
+			OpenConnection();
+		}
+	}
+}
+
 FText FLiveLinkPrestonMDRSource::GetSourceStatus() const
 {
-	if (bFailedToConnectToDevice == true)
+	switch (SourceStatus)
+	{
+	case EPrestonSourceStatus::NotConnected:
+	{
+		return LOCTEXT("NotConnectedStatus", "Not Connected");
+	}
+	case EPrestonSourceStatus::WaitingToConnect:
 	{
-		return LOCTEXT("FailedConnectionStatus", "Failed to connect");
+		return LOCTEXT("WaitingToConnectStatus", "Waiting To Connect...");
 	}
-	else if (bIsConnectedToDevice == false)
+	case EPrestonSourceStatus::ConnectedActive:
 	{
-		return LOCTEXT("NotConnectedStatus", "Waiting to connect...");
+		return LOCTEXT("ConnectedActiveStatus", "Active");
 	}
-	else if (FPlatformTime::Seconds() - LastTimeDataReceived > DataReceivedTimeout)
+	case EPrestonSourceStatus::ConnectedIdle:
 	{
-		return LOCTEXT("WaitingForDataStatus", "Idle");
+		return LOCTEXT("ConnectedIdleStatus", "Idle");
 	}
-	return LOCTEXT("ActiveStatus", "Active");
+	case EPrestonSourceStatus::ConnectionFailed:
+	{
+		return LOCTEXT("ConnectionFailedStatus", "Failed To Connect");
+	}
+	case EPrestonSourceStatus::ConnectionLost:
+	{
+		return LOCTEXT("ConnectionLostStatus", "Connection Lost");
+	}
+	case EPrestonSourceStatus::ShuttingDown:
+	{
+		return LOCTEXT("ShuttingDownStatus", "Shutting Down...");
+	}
+	default:
+	{
+		return LOCTEXT("UnknownStatus", "Unknown");
+	}
+	}
+
+	return LOCTEXT("UnknownStatus", "Unknown");
 }
 
 void FLiveLinkPrestonMDRSource::OpenConnection()
@@ -185,69 +253,63 @@ void FLiveLinkPrestonMDRSource::OpenConnection()
 	Socket->SetRecvErr(true);
 
 	FIPv4Address IPAddr;
-	if (FIPv4Address::Parse(ConnectionSettings.IPAddress, IPAddr) == false)
+	if (!FIPv4Address::Parse(ConnectionSettings.IPAddress, IPAddr))
 	{
-		UE_LOG(LogPrestonMDRSource, Error, TEXT("Ill-formed IP Address"));
+		SourceStatus = EPrestonSourceStatus::ConnectionFailed;
+		UE_LOG(LogLiveLinkPrestonMDRSource, Error, TEXT("Ill-formed IP Address"));
 		return;
 	}
 
-	uint16 PortNumber = ConnectionSettings.PortNumber;
-
-	FIPv4Endpoint Endpoint = FIPv4Endpoint(IPAddr, PortNumber);
+	FIPv4Endpoint Endpoint = FIPv4Endpoint(IPAddr, ConnectionSettings.PortNumber);
 	TSharedRef<FInternetAddr> Addr = Endpoint.ToInternetAddr();
 
-	UE_LOG(LogPrestonMDRSource, VeryVerbose, TEXT("Connecting to the MDR server at %s:%d..."), *IPAddr.ToString(), PortNumber);
+	UE_LOG(LogLiveLinkPrestonMDRSource, VeryVerbose, TEXT("Connecting to the MDR server at %s"), *Endpoint.ToString());
 
 	if (!Socket->Connect(*Addr))
 	{
-		UE_LOG(LogPrestonMDRSource, Error, TEXT("Could not connect to the MDR server at %s:%d"), *IPAddr.ToString(), PortNumber);
-		bFailedToConnectToDevice = true;
+		SourceStatus = EPrestonSourceStatus::ConnectionFailed;
+		UE_LOG(LogLiveLinkPrestonMDRSource, Error, TEXT("Could not connect to the MDR server at %s"), *Endpoint.ToString());
 		return;
 	}
 
 	if (Socket->GetConnectionState() == ESocketConnectionState::SCS_ConnectionError)
 	{
-		UE_LOG(LogPrestonMDRSource, Error, TEXT("Could not connect to the MDR server at %s:%d"), *IPAddr.ToString(), PortNumber);
-		bFailedToConnectToDevice = true;
+		SourceStatus = EPrestonSourceStatus::ConnectionFailed;
+		UE_LOG(LogLiveLinkPrestonMDRSource, Error, TEXT("Could not connect to the MDR server at %s"), *Endpoint.ToString());
 		return;
 	}
+
+	// The socket is non-blocking, so we will not know if the connection was successful until sometime later
+	SourceStatus = EPrestonSourceStatus::WaitingToConnect;
+
+	MessageThread = MakeUnique<FPrestonMDRMessageThread>(Socket);
+
+	MessageThread->OnFrameDataReady().BindRaw(this, &FLiveLinkPrestonMDRSource::OnFrameDataReady);
+	MessageThread->OnStatusChanged().BindRaw(this, &FLiveLinkPrestonMDRSource::OnStatusChanged);
+	MessageThread->OnConnectionLost().BindRaw(this, &FLiveLinkPrestonMDRSource::OnConnectionLost);
+	MessageThread->OnConnectionFailed().BindRaw(this, &FLiveLinkPrestonMDRSource::OnConnectionFailed);
+
+	MessageThread->Start();
 }
 
 void FLiveLinkPrestonMDRSource::OnConnectionLost()
 {
-	UE_LOG(LogPrestonMDRSource, Error, TEXT("Connection to the MDR device was lost"));
-	bIsConnectedToDevice = false;
-
 	{
 		FScopeLock Lock(&PrestonSourceCriticalSection);
-		// There was an unrecoverable connection loss, so we need to destroy the current socket and open a new one
-		if (Socket)
-		{
-			SocketSubsystem->DestroySocket(Socket);
-			Socket = nullptr;
-		}
-
-		// Attempt to re-open the connection to the MDR server
-		OpenConnection();
+		SourceStatus = EPrestonSourceStatus::ConnectionLost;
 	}
 
-	if (MessageThread && Socket)
-	{
-		MessageThread->SetSocket(Socket);
-		MessageThread->SoftReset();
-	}
+	UE_LOG(LogLiveLinkPrestonMDRSource, Error, TEXT("Connection to the MDR device was lost"));	
 }
 
 void FLiveLinkPrestonMDRSource::OnConnectionFailed()
 {
-	UE_LOG(LogPrestonMDRSource, Error, TEXT("Could not connect to the MDR server at %s:%d"), *ConnectionSettings.IPAddress, ConnectionSettings.PortNumber);
-	bIsConnectedToDevice = false;
-	bFailedToConnectToDevice = true;
-
-	if (MessageThread)
 	{
-		MessageThread->Stop();
+		FScopeLock Lock(&PrestonSourceCriticalSection);
+		SourceStatus = EPrestonSourceStatus::ConnectionFailed;
 	}
+
+	UE_LOG(LogLiveLinkPrestonMDRSource, Error, TEXT("Could not connect to the MDR server at %s:%d"), *ConnectionSettings.IPAddress, ConnectionSettings.PortNumber);
 }
 
 void FLiveLinkPrestonMDRSource::OnStatusChanged(FMDR3Status InStatus)
@@ -274,7 +336,10 @@ void FLiveLinkPrestonMDRSource::UpdateStaticData()
 
 void FLiveLinkPrestonMDRSource::OnFrameDataReady(FLensDataPacket InData)
 {
-	bIsConnectedToDevice = true;
+	{
+		FScopeLock Lock(&PrestonSourceCriticalSection);
+		SourceStatus = EPrestonSourceStatus::ConnectedActive;
+	}
 
 	FLiveLinkFrameDataStruct LensFrameDataStruct(FLiveLinkPrestonMDRFrameData::StaticStruct());
 	FLiveLinkPrestonMDRFrameData* LensFrameData = LensFrameDataStruct.Cast<FLiveLinkPrestonMDRFrameData>();
diff --git a/Engine/Plugins/VirtualProduction/LiveLinkPrestonMDR/Source/LiveLinkPrestonMDR/Private/LiveLinkPrestonMDRSource.h b/Engine/Plugins/VirtualProduction/LiveLinkPrestonMDR/Source/LiveLinkPrestonMDR/Private/LiveLinkPrestonMDRSource.h
--- a/Engine/Plugins/VirtualProduction/LiveLinkPrestonMDR/Source/LiveLinkPrestonMDR/Private/LiveLinkPrestonMDRSource.h
+++ b/Engine/Plugins/VirtualProduction/LiveLinkPrestonMDR/Source/LiveLinkPrestonMDR/Private/LiveLinkPrestonMDRSource.h
@@ -14,6 +14,17 @@
 
 class ISocketSubsystem;
 
+enum class EPrestonSourceStatus : uint8
+{
+	NotConnected,
+	WaitingToConnect,
+	ConnectedActive,
+	ConnectedIdle,
+	ConnectionFailed,
+	ConnectionLost,
+	ShuttingDown
+};
+
 class LIVELINKPRESTONMDR_API FLiveLinkPrestonMDRSource : public ILiveLinkSource
 {
 public:
@@ -23,7 +34,7 @@ public:
 
 	// Begin ILiveLinkSource Implementation
 	virtual void ReceiveClient(ILiveLinkClient* InClient, FGuid InSourceGuid) override;
-
+	virtual void Update() override;
 	virtual bool IsSourceStillValid() const override;
 
 	virtual bool RequestSourceShutdown() override;
@@ -45,6 +56,7 @@ private:
 	void OnConnectionLost();
 	void OnConnectionFailed();
 	void UpdateStaticData();
+	bool ShutdownMessageThreadAndSocket();
 
 	ILiveLinkClient* Client = nullptr;
 	FSocket* Socket = nullptr;
@@ -54,19 +66,20 @@ private:
 	FLiveLinkSubjectKey SubjectKey;
 	FText SourceMachineName;
 
-	ULiveLinkPrestonMDRSourceSettings* SavedSourceSettings = nullptr;
+	TObjectPtr<ULiveLinkPrestonMDRSourceSettings> SavedSourceSettings = nullptr;
 
 	TUniquePtr<FPrestonMDRMessageThread> MessageThread;
 
 	FCriticalSection PrestonSourceCriticalSection;
 
 	std::atomic<double> LastTimeDataReceived;
+	double MessageThreadShutdownTime;
 
-	std::atomic<bool> bIsConnectedToDevice;
-	std::atomic<bool> bFailedToConnectToDevice;
+	EPrestonSourceStatus SourceStatus = EPrestonSourceStatus::NotConnected;
 
 	FMDR3Status LatestMDRStatus;
 
 private:
 	const float DataReceivedTimeout = 1.0f;
+	const float MessageThreadShutdownTimeout = 2.0f;
 };
\ No newline at end of file
diff --git a/Engine/Plugins/VirtualProduction/LiveLinkPrestonMDR/Source/LiveLinkPrestonMDR/Private/PrestonMDRMessageThread.cpp b/Engine/Plugins/VirtualProduction/LiveLinkPrestonMDR/Source/LiveLinkPrestonMDR/Private/PrestonMDRMessageThread.cpp
--- a/Engine/Plugins/VirtualProduction/LiveLinkPrestonMDR/Source/LiveLinkPrestonMDR/Private/PrestonMDRMessageThread.cpp
+++ b/Engine/Plugins/VirtualProduction/LiveLinkPrestonMDR/Source/LiveLinkPrestonMDR/Private/PrestonMDRMessageThread.cpp
@@ -49,6 +49,8 @@ void FPrestonMDRMessageThread::SetIncomingDataMode_GameThread(EFIZDataMode InDat
 
 void FPrestonMDRMessageThread::Start()
 {
+	bIsFinished = false;
+
 	Thread.Reset(FRunnableThread::Create(this, TEXT("Preston MDR Message Thread"), ThreadStackSize, TPri_AboveNormal));
 }
 
@@ -88,6 +90,7 @@ uint32 FPrestonMDRMessageThread::Run()
 		if ((FPlatformTime::Seconds() - LastTimeDataReceived) > ConnectionTimeout)
 		{
 			ConnectionFailedDelegate.ExecuteIfBound();
+			bIsFinished = true;
 			return 0;
 		}
 
@@ -99,12 +102,17 @@ uint32 FPrestonMDRMessageThread::Run()
 
 	int32 TotalBytesRead = 0;
 	uint32 NumBytesOnSocket = 0;
-	while ((bIsThreadRunning) || Socket->HasPendingData(NumBytesOnSocket))
+
+	uint32 NumCommandsSent = 0;
+	uint32 NumRepliedReceived = 0;
+
+	while ((bIsThreadRunning || !bReadyToSendCommandToServer) && !bForceKillThread)
 	{
 		if (bReadyToSendCommandToServer)
 		{
 			bReadyToSendCommandToServer = false;
 			SendNextCommandToServer();
+			NumCommandsSent++;
 		}
 
 		int32 NumBytesRead = 0;
@@ -138,6 +146,7 @@ uint32 FPrestonMDRMessageThread::Run()
 				LastTimeDataReceived = FPlatformTime::Seconds();
 
 				ParseMessageFromServer(ReceiveBuffer.GetData(), PacketLength);
+				NumRepliedReceived++;
 
 				// After receiving a new timecode and new FIZ data packet, send a new frame of lens data 
 				if (bIsFrameTimeReady && bIsFrameDataReady)
@@ -159,7 +168,7 @@ uint32 FPrestonMDRMessageThread::Run()
 		}
 		else if (Socket->GetConnectionState() != ESocketConnectionState::SCS_Connected)
 		{
-			ConnectionFailedDelegate.ExecuteIfBound();
+			ConnectionLostDelegate.ExecuteIfBound();
 			break;
 		}
 
@@ -181,6 +190,11 @@ uint32 FPrestonMDRMessageThread::Run()
 		}
 	}
 
+	UE_LOG(LogPrestonMDRMessageThread, Verbose, TEXT("Commands Sent   : %d"), NumCommandsSent);
+	UE_LOG(LogPrestonMDRMessageThread, Verbose, TEXT("Replies Received: %d"), NumRepliedReceived);
+
+	bIsFinished = true;
+
 	return 0;
 }
 
@@ -584,11 +598,17 @@ void FPrestonMDRMessageThread::SoftReset()
 
 void FPrestonMDRMessageThread::HardReset()
 {
+	// Update this flag so that the main message loop stops executing and the thread can be marked as finished
+	bForceKillThread = true;
+
 	// Notify the Source that the connection has been lost so that it can attempt to re-establish it.
 	ConnectionLostDelegate.ExecuteIfBound();
+}
 
-	// Reset the timer so that another reset does not immediately trigger
-	LastTimeDataReceived = FPlatformTime::Seconds();
+void FPrestonMDRMessageThread::ForceKill()
+{
+	// Update this flag so that the main message loop stops executing and the thread can be marked as finished
+	bForceKillThread = true;
 }
 
 // Utility Functions
diff --git a/Engine/Plugins/VirtualProduction/LiveLinkPrestonMDR/Source/LiveLinkPrestonMDR/Private/PrestonMDRMessageThread.h b/Engine/Plugins/VirtualProduction/LiveLinkPrestonMDR/Source/LiveLinkPrestonMDR/Private/PrestonMDRMessageThread.h
--- a/Engine/Plugins/VirtualProduction/LiveLinkPrestonMDR/Source/LiveLinkPrestonMDR/Private/PrestonMDRMessageThread.h
+++ b/Engine/Plugins/VirtualProduction/LiveLinkPrestonMDR/Source/LiveLinkPrestonMDR/Private/PrestonMDRMessageThread.h
@@ -168,6 +168,9 @@ public:
 
 	void Start();
 
+	bool IsThreadRunning() const { return bIsThreadRunning; }
+	bool IsFinished() const { return bIsFinished; }
+
 	/**
 	 * Returns a delegate that is executed when frame data is ready.
 	 *
@@ -205,6 +208,8 @@ public:
 
 	void SoftReset();
 
+	void ForceKill();
+
 	void SetIncomingDataMode_GameThread(EFIZDataMode InDataMode);
 
 private:
@@ -236,7 +241,7 @@ private:
 	FSocket* Socket;
 
 	TUniquePtr<FRunnableThread> Thread;
-	bool bIsThreadRunning = false;
+	std::atomic<bool> bIsThreadRunning = false;
 
 	FOnFrameDataReady FrameDataReadyDelegate;
 	FOnStatusChanged StatusChangedDelegate;
@@ -269,12 +274,15 @@ private:
 	bool bSoftResetTriggered = false;
 	bool bHardResetTriggered = false;
 
+	std::atomic<bool> bIsFinished = true;
+	std::atomic<bool> bForceKillThread = false;
+
 private:
 	static constexpr uint32 MaximumMDRMessageLength = 512;
 
 	static constexpr uint32 ThreadStackSize = 1024 * 128;
 
 	static constexpr float ConnectionWaitInterval = 0.02f; // 20ms
-	static constexpr double ConnectionTimeout = 3.0; // 3sec
+	static constexpr double ConnectionTimeout = 5.0; // 3sec
 	static constexpr double DataReceivedTimeout = 2.0; // 2sec
 };
